Fixed unsafe type cast warnings in SynthBenchmark module.

diff --git a/Engine/Source/Runtime/SynthBenchmark/Private/RayIntersectBenchmark.cpp b/Engine/Source/Runtime/SynthBenchmark/Private/RayIntersectBenchmark.cpp
--- a/Engine/Source/Runtime/SynthBenchmark/Private/RayIntersectBenchmark.cpp
+++ b/Engine/Source/Runtime/SynthBenchmark/Private/RayIntersectBenchmark.cpp
@@ -21,32 +21,32 @@ FORCEINLINE bool LineCheckWithTriangle(const FVector& V1, const FVector& V2, con
 	FVector	Edge1 = V3 - V1;
 	FVector	Edge2 = V2 - V1;
 	FVector	P = Direction ^ Edge2;
-	float Determinant = Edge1 | P;
+	double Determinant = Edge1 | P;
 
-	if(Determinant < 0.00001f)
+	if(Determinant < 0.00001)
 	{
 		return false;
 	}
 
 	FVector	T = Start - V1;
-	float U = T | P;
+	double U = T | P;
 
-	if(U < 0.0f || U > Determinant)
+	if(U < 0.0 || U > Determinant)
 	{
 		return false;
 	}
 
 	FVector	Q = T ^ Edge1;
-	float V = Direction | Q;
+	double V = Direction | Q;
 
-	if(V < 0.0f || U + V > Determinant)
+	if(V < 0.0 || U + V > Determinant)
 	{
 		return false;
 	}
 
-	float Time = (Edge2 | Q) / Determinant;
+	double Time = (Edge2 | Q) / Determinant;
 
-	return Time >= 0.0f;
+	return Time >= 0.0;
 }
 
 
diff --git a/Engine/Source/Runtime/SynthBenchmark/Private/SynthBenchmarkPrivate.cpp b/Engine/Source/Runtime/SynthBenchmark/Private/SynthBenchmarkPrivate.cpp
--- a/Engine/Source/Runtime/SynthBenchmark/Private/SynthBenchmarkPrivate.cpp
+++ b/Engine/Source/Runtime/SynthBenchmark/Private/SynthBenchmarkPrivate.cpp
@@ -263,7 +263,7 @@ void FSynthBenchmark::GetRHIInfo(FGPUAdpater& Out, FString& RHIName) const
 	{
 		FTextureMemoryStats TextureMemStats;
 		RHIGetTextureMemoryStats(TextureMemStats);
-		WriteFStringToResults(Out.AdapterDedicatedMemoryMB, FString::FromInt(TextureMemStats.DedicatedVideoMemory / (1024*1024)));
+		WriteFStringToResults(Out.AdapterDedicatedMemoryMB, FString::FromInt(static_cast<int32>(TextureMemStats.DedicatedVideoMemory / (1024*1024))));
 
 		RHIName = GDynamicRHI->GetName();
 	}
