[SynthBenchmark] Removed public dependency on Renderer to break circular dependency "SynthBenchmark -> Renderer -> Engine -> SynthBenchmark"

diff --git a/Engine/Source/Runtime/SynthBenchmark/SynthBenchmark.Build.cs b/Engine/Source/Runtime/SynthBenchmark/SynthBenchmark.Build.cs
--- a/Engine/Source/Runtime/SynthBenchmark/SynthBenchmark.Build.cs
+++ b/Engine/Source/Runtime/SynthBenchmark/SynthBenchmark.Build.cs
@@ -6,8 +6,6 @@ namespace UnrealBuildTool.Rules
 	{
 		public SynthBenchmark(ReadOnlyTargetRules Target) : base(Target)
 		{
-            PublicIncludePathModuleNames.Add("Renderer");
-
 			PublicDependencyModuleNames.AddRange(
 				new string[]
 				{
