SynthBenchmark: Fixed conversion warnings. Enabled UnsafeTypeCast warnings as errors.

diff --git a/Engine/Source/Runtime/SynthBenchmark/Private/FractalBenchmark.cpp b/Engine/Source/Runtime/SynthBenchmark/Private/FractalBenchmark.cpp
--- a/Engine/Source/Runtime/SynthBenchmark/Private/FractalBenchmark.cpp
+++ b/Engine/Source/Runtime/SynthBenchmark/Private/FractalBenchmark.cpp
@@ -46,7 +46,7 @@ float FractalBenchmark()
 	{
 		for(int x = 0; x < Extent; ++x)
 		{
-			Sum += EvaluateJuliaFractalAt(x / (float)Extent * 2 - 1, y / (float)Extent * 2 - 1);
+			Sum += (float)EvaluateJuliaFractalAt((float)x / (float)Extent * 2.0f - 1.0f, (float)y / (float)Extent * 2.0f - 1.0f);
 		}
 	}
 	
diff --git a/Engine/Source/Runtime/SynthBenchmark/Private/RayIntersectBenchmark.cpp b/Engine/Source/Runtime/SynthBenchmark/Private/RayIntersectBenchmark.cpp
--- a/Engine/Source/Runtime/SynthBenchmark/Private/RayIntersectBenchmark.cpp
+++ b/Engine/Source/Runtime/SynthBenchmark/Private/RayIntersectBenchmark.cpp
@@ -72,5 +72,5 @@ float RayIntersectBenchmark()
 	}
 
 	// to avoid getting optimized out
-	return HitCount / (float)StepCount;
+	return (float)HitCount / (float)StepCount;
 }
diff --git a/Engine/Source/Runtime/SynthBenchmark/Private/SynthBenchmarkPrivate.cpp b/Engine/Source/Runtime/SynthBenchmark/Private/SynthBenchmarkPrivate.cpp
--- a/Engine/Source/Runtime/SynthBenchmark/Private/SynthBenchmarkPrivate.cpp
+++ b/Engine/Source/Runtime/SynthBenchmark/Private/SynthBenchmarkPrivate.cpp
@@ -63,7 +63,7 @@ static FTimeSample RunBenchmark(float WorkScale, float (*Function)())
 		FPlatformMisc::MemoryBarrier();
 	}
 	
-	return FTimeSample(Sum, Sum / RunCount);
+	return FTimeSample(Sum, Sum / (float)RunCount);
 }
 
 template<int NumMethods>
diff --git a/Engine/Source/Runtime/SynthBenchmark/SynthBenchmark.Build.cs b/Engine/Source/Runtime/SynthBenchmark/SynthBenchmark.Build.cs
--- a/Engine/Source/Runtime/SynthBenchmark/SynthBenchmark.Build.cs
+++ b/Engine/Source/Runtime/SynthBenchmark/SynthBenchmark.Build.cs
@@ -6,15 +6,17 @@ namespace UnrealBuildTool.Rules
 	{
 		public SynthBenchmark(ReadOnlyTargetRules Target) : base(Target)
 		{
+			UnsafeTypeCastWarningLevel = WarningLevel.Error;
+
 			PublicDependencyModuleNames.AddRange(
 				new string[]
 				{
 					"Core",
-        			"RenderCore", 
-				    "RHI",
+					"RenderCore",
+					"RHI",
 					// ... add other public dependencies that you statically link with here ...
 				}
-				);
+			);
 
 			PrivateDependencyModuleNames.AddRange(
 				new string[]
@@ -22,14 +24,14 @@ namespace UnrealBuildTool.Rules
 					"ApplicationCore"
 					// ... add private dependencies that you statically link with here ...
 				}
-				);
+			);
 
 			DynamicallyLoadedModuleNames.AddRange(
 				new string[]
 				{
 					// ... add any modules that your module loads dynamically here ...
 				}
-				);
+			);
 		}
 	}
 }
