Modify GPU synth benchmark to do some warmup rounds at a configurable work scale and an optional full test.

#tested QA tested on several low-end and high-end machines

diff --git a/Engine/Source/Runtime/SynthBenchmark/Private/SynthBenchmarkPrivate.cpp b/Engine/Source/Runtime/SynthBenchmark/Private/SynthBenchmarkPrivate.cpp
--- a/Engine/Source/Runtime/SynthBenchmark/Private/SynthBenchmarkPrivate.cpp
+++ b/Engine/Source/Runtime/SynthBenchmark/Private/SynthBenchmarkPrivate.cpp
@@ -8,6 +8,31 @@
 #include "RHI.h"
 #include "RendererInterface.h"
 
+int32 GSynthBenchmarkGPUWarmupRounds = 3;
+static FAutoConsoleVariableRef CVarSynthBenchmarkGPUWarmupRounds(
+	TEXT("r.SynthBenchmark.GPUWarmupRounds"),
+	GSynthBenchmarkGPUWarmupRounds,
+	TEXT("The number of warmup rounds for the GPU synthetic benchmark (default 3)."),
+	ECVF_Default
+);
+
+float GSynthBenchmarkGPUWarmupWorkScale = 0.15f;
+static FAutoConsoleVariableRef CVarSynthBenchmarkGPUWarmupWorkScale(
+	TEXT("r.SynthBenchmark.GPUWarmupWorkScale"),
+	GSynthBenchmarkGPUWarmupWorkScale,
+	TEXT("The workload scale applied to the GPU synthetic benchmark's warmup rounds (default 0.15f)."),
+	ECVF_Default
+);
+
+float GSynthBenchmarkGPUFullTestThreshold = 1.0f;
+static FAutoConsoleVariableRef CVarSynthBenchmarkGPUFullTestThreshold(
+	TEXT("r.SynthBenchmark.GPUWarmupThreshold"),
+	GSynthBenchmarkGPUFullTestThreshold,
+	TEXT("Controls the maximum time, in seconds, used to determine whether to run the full GPU synthetic benchmark workload (default 1 second).\n")
+	TEXT("If the last warmup round takes longer than this threshold, the full GPU test is not run to avoid the risk of triggering the GPU hang detection (TDR)."),
+	ECVF_Default
+);
+
 float RayIntersectBenchmark();
 float FractalBenchmark();
 
@@ -173,44 +198,23 @@ void FSynthBenchmark::Run(FSynthBenchmarkResults& InOut, bool bGPUBenchmark, flo
 		{
 			IRendererModule& RendererModule = FModuleManager::GetModuleChecked<IRendererModule>(TEXT("Renderer"));
 
-			// First we run a quick test. If that shows very bad performance we don't need another test
-			// The hardware is slow, we don't need a long test and risk driver TDR (driver recovery).
-			// We have seen this problem on very low end GPUs.
+			// Do a few rounds for warmup.
+			for (int32 Round = 0; Round < GSynthBenchmarkGPUWarmupRounds; ++Round)
 			{
-				const float fFirstWorkScale = 0.01f * WorkScale;
-				const float fSecondWorkScale = 0.1f * WorkScale;
-
-				RendererModule.GPUBenchmark(InOut, fFirstWorkScale);
+				RendererModule.GPUBenchmark(InOut, WorkScale * GSynthBenchmarkGPUWarmupWorkScale);
 				GPUTime = InOut.ComputeTotalGPUTime();
-				if (GPUTime > 0.0f)
-				{
-					UE_LOG(LogSynthBenchmark, Display, TEXT("  GPU first test: %.2fs"), GPUTime);
-					PrintGPUStats(InOut.GPUStats, TEXT(" (likely to be very inaccurate)"));
-				}
+				UE_LOG(LogSynthBenchmark, Display, TEXT("  GPU warmup (round %d) test: %.2fs"), Round + 1, GPUTime);
+				PrintGPUStats(InOut.GPUStats, TEXT(" (likely to be inaccurate)"));
+			}
 
-				if (GPUTime < 0.15f)
-				{
-					RendererModule.GPUBenchmark(InOut, fSecondWorkScale);
-					GPUTime = InOut.ComputeTotalGPUTime();
-
-					if (GPUTime > 0.0f)
-					{
-						UE_LOG(LogSynthBenchmark, Display, TEXT("  GPU second test: %.2fs"), GPUTime);
-						PrintGPUStats(InOut.GPUStats, TEXT(" (likely to be inaccurate)"));
-					}
-
-					if (GPUTime < 0.15f)
-					{
-						RendererModule.GPUBenchmark(InOut, WorkScale);
-						GPUTime = InOut.ComputeTotalGPUTime();
-
-						if (GPUTime > 0.0f)
-						{
-							UE_LOG(LogSynthBenchmark, Display, TEXT("  GPU third test: %.2fs"), GPUTime);
-							PrintGPUStats(InOut.GPUStats, TEXT(""));
-						}
-					}
-				}
+			// If even after a few rounds of warmup we're still taking a long time just doing a portion
+			// of the work, don't risk a TDR doing the full test.
+			if (GPUTime < GSynthBenchmarkGPUFullTestThreshold)
+			{
+				RendererModule.GPUBenchmark(InOut, WorkScale);
+				GPUTime = InOut.ComputeTotalGPUTime();
+				UE_LOG(LogSynthBenchmark, Display, TEXT("  GPU final test: %.2fs"), GPUTime);
+				PrintGPUStats(InOut.GPUStats, TEXT(""));
 			}
 
 			if (GPUTime > 0.0f)
